[RigLogicMutable] Moved plugin to experimental.

- Tagged DNA Constant node as experimental.

#rnx

diff --git a/Engine/Plugins/Experimental/RigLogicMutable/RigLogicMutable.uplugin b/Engine/Plugins/Experimental/RigLogicMutable/RigLogicMutable.uplugin
new file mode 100644
--- /dev/null
+++ b/Engine/Plugins/Experimental/RigLogicMutable/RigLogicMutable.uplugin
@@ -0,0 +1,43 @@
+{
+	"FileVersion": 3,
+	"Version": 1,
+	"VersionName": "0.1",
+	"FriendlyName": "RigLogic Extensions For Mutable",
+	"Description": "Adds Mutable functionality to work with RigLogic DNA",
+	"Category": "Animation",
+	"CreatedBy": "Epic Games, Inc.",
+	"CreatedByURL": "https://epicgames.com",
+	"DocsURL": "",
+	"MarketplaceURL": "",
+	"SupportURL": "",
+	"CanContainContent": false,
+	"IsBetaVersion": false,
+	"IsExperimentalVersion": true,
+	"Installed": false,
+	"EnabledByDefault": false,
+	"Modules": [
+		{
+			"Name": "RigLogicMutable",
+			"Type": "Runtime",
+			"LoadingPhase": "Default"
+		},
+		{
+			"Name": "RigLogicMutableEditor",
+			"Type": "Editor",
+			"LoadingPhase": "Default"
+		}
+
+	],
+	"Plugins": [
+		{
+			"Name": "RigLogic",
+			"Enabled": true
+		},
+		{
+			"Name": "Mutable",
+			"Enabled": true
+		}
+	]
+
+
+}
\ No newline at end of file
diff --git a/Engine/Plugins/Experimental/RigLogicMutable/Source/RigLogicMutable/Private/RigLogicMutableExtension.cpp b/Engine/Plugins/Experimental/RigLogicMutable/Source/RigLogicMutable/Private/RigLogicMutableExtension.cpp
new file mode 100644
--- /dev/null
+++ b/Engine/Plugins/Experimental/RigLogicMutable/Source/RigLogicMutable/Private/RigLogicMutableExtension.cpp
@@ -0,0 +1,123 @@
+// Copyright Epic Games, Inc. All Rights Reserved.
+
+#include "RigLogicMutableExtension.h"
+
+#include "DNAAsset.h"
+#include "Engine/SkeletalMesh.h"
+#include "MuCO/CustomizableObject.h"
+#include "Serialization/MemoryWriter.h"
+#include "Serialization/MemoryReader.h"
+#include "UObject/UObjectGlobals.h"
+
+const FName URigLogicMutableExtension::DNAPinType(TEXT("DNA"));
+const FName URigLogicMutableExtension::DNABaseNodePinName(TEXT("DNA"));
+const FText URigLogicMutableExtension::DNANodeCategory(FText::FromString(TEXT("DNA")));
+
+FDNAPinData::FDNAPinData(FDNAPinData&& Source)
+{
+	// Invoke assignment operator
+	*this = MoveTemp(Source);
+}
+
+FDNAPinData& FDNAPinData::operator=(FDNAPinData&& Source)
+{
+	ComponentIndex = Source.ComponentIndex;
+
+	DNAAsset = Source.DNAAsset;
+	Source.DNAAsset = nullptr;
+
+	return *this;
+}
+
+void FDNAPinData::CopyFromDNAAsset(const UDNAAsset* SourceAsset, UObject* OuterForOwnedObjects)
+{
+	if (SourceAsset)
+	{
+		DNAAsset = URigLogicMutableExtension::CopyDNAAsset(SourceAsset, OuterForOwnedObjects);
+	}
+	else
+	{
+		DNAAsset = nullptr;
+	}
+}
+
+TArray<FCustomizableObjectPinType> URigLogicMutableExtension::GetPinTypes() const
+{
+	TArray<FCustomizableObjectPinType> Result;
+	
+	FCustomizableObjectPinType& DNAType = Result.AddDefaulted_GetRef();
+	DNAType.Name = DNAPinType;
+	DNAType.DisplayName = FText::FromString(TEXT("RigLogic DNA"));
+	DNAType.Color = FLinearColor::Red;
+
+	return Result;
+}
+
+TArray<FObjectNodeInputPin> URigLogicMutableExtension::GetAdditionalObjectNodePins() const
+{
+	TArray<FObjectNodeInputPin> Result;
+
+	FObjectNodeInputPin& DNAInputPin = Result.AddDefaulted_GetRef();
+	DNAInputPin.PinType = DNAPinType;
+	DNAInputPin.PinName = DNABaseNodePinName;
+	DNAInputPin.DisplayName = FText::FromString(TEXT("RigLogic DNA"));
+	DNAInputPin.bIsArray = false;
+
+	return Result;
+}
+
+void URigLogicMutableExtension::OnSkeletalMeshCreated(
+	const TArray<FInputPinDataContainer>& InputPinData,
+	int32 ComponentIndex,
+	USkeletalMesh* SkeletalMesh) const
+{
+	// Find the DNA produced by the Customizable Object, if any, and assign it to the Skeletal Mesh
+
+	for (const FInputPinDataContainer& Container : InputPinData)
+	{
+		if (Container.Pin.PinName == DNABaseNodePinName)
+		{
+			const FDNAPinData* Data = Container.Data.GetPtr<FDNAPinData>();
+			if (Data
+				&& Data->ComponentIndex == ComponentIndex
+				&& Data->GetDNAAsset() != nullptr)
+			{
+				UDNAAsset* NewDNA = CopyDNAAsset(Data->GetDNAAsset(), SkeletalMesh);
+				SkeletalMesh->AddAssetUserData(NewDNA);
+
+				// A mesh can only have one DNA at a time, so if the CO produced multiple DNA
+				// Assets, all but the first will be discarded.
+				break;
+			}
+		}
+	}
+}
+
+UDNAAsset* URigLogicMutableExtension::CopyDNAAsset(const UDNAAsset* Source, UObject* OuterForCopy)
+{
+	check(IsInGameThread());
+
+	// Currently the only way to copy a UDNAAsset is to serialize it into a buffer and deserialize
+	// the buffer into the copy.
+
+	// Serialize the existing DNA into a buffer
+	TArray<uint8> Buffer;
+	{
+		FMemoryWriter Writer(Buffer);
+
+		// Need to const_cast because Serialize is non-const. Source will not be modified.
+		const_cast<UDNAAsset*>(Source)->Serialize(Writer);
+
+		Writer.Close();
+	}
+
+	// Create the new DNA asset and deserialize the buffer into it
+	UDNAAsset* Result = NewObject<UDNAAsset>(OuterForCopy);
+	{
+		FMemoryReader Reader(Buffer);
+
+		Result->Serialize(Reader);
+	}
+
+	return Result;
+}
diff --git a/Engine/Plugins/Experimental/RigLogicMutable/Source/RigLogicMutable/Private/RigLogicMutableModule.cpp b/Engine/Plugins/Experimental/RigLogicMutable/Source/RigLogicMutable/Private/RigLogicMutableModule.cpp
new file mode 100644
--- /dev/null
+++ b/Engine/Plugins/Experimental/RigLogicMutable/Source/RigLogicMutable/Private/RigLogicMutableModule.cpp
@@ -0,0 +1,29 @@
+// Copyright Epic Games, Inc. All Rights Reserved.
+
+#include "RigLogicMutableModule.h"
+
+#include "MuCO/ICustomizableObjectModule.h"
+#include "RigLogicMutableExtension.h"
+
+class FRigLogicMutableModule : public IModuleInterface
+{
+	virtual void StartupModule() override
+	{
+		Extension = TWeakObjectPtr<const URigLogicMutableExtension>(GetDefault<URigLogicMutableExtension>());
+		ICustomizableObjectModule::Get().RegisterExtension(Extension.Get());
+	}
+
+	virtual void ShutdownModule() override
+	{
+		if (ICustomizableObjectModule::IsAvailable()
+			&& Extension.IsValid())
+		{
+			ICustomizableObjectModule::Get().UnregisterExtension(Extension.Get());
+		}
+	}
+
+private:
+	TWeakObjectPtr<const URigLogicMutableExtension> Extension;
+};
+
+IMPLEMENT_MODULE(FRigLogicMutableModule, RigLogicMutable);
diff --git a/Engine/Plugins/Experimental/RigLogicMutable/Source/RigLogicMutable/Public/RigLogicMutableExtension.h b/Engine/Plugins/Experimental/RigLogicMutable/Source/RigLogicMutable/Public/RigLogicMutableExtension.h
new file mode 100644
--- /dev/null
+++ b/Engine/Plugins/Experimental/RigLogicMutable/Source/RigLogicMutable/Public/RigLogicMutableExtension.h
@@ -0,0 +1,88 @@
+// Copyright Epic Games, Inc. All Rights Reserved.
+
+#pragma once
+
+#include "MuCO/CustomizableObjectExtension.h"
+#include "UObject/SoftObjectPtr.h"
+
+#include "RigLogicMutableExtension.generated.h"
+
+class UDNAAsset;
+
+/** Used as ExtensionData to represent a DNA Asset in a Customizable Object graph */
+USTRUCT()
+struct RIGLOGICMUTABLE_API FDNAPinData
+{
+	GENERATED_BODY()
+
+public:
+	FDNAPinData() = default;
+
+	/**
+	 * Direct copying is not allowed, because the DNA Asset can only be owned by one struct.
+	 *
+	 * To make a copy of this, create a new default instance and call CopyFromDNAAsset on it to
+	 * copy the DNA Asset into it.
+	 */
+	FDNAPinData(const FDNAPinData& Source) = delete;
+	FDNAPinData& operator=(const FDNAPinData& Source) = delete;
+
+	/** Move constructor transfers ownership of the DNA */
+	FDNAPinData(FDNAPinData&& Source);
+	FDNAPinData& operator=(FDNAPinData&& Source);
+
+	/**
+	 * Makes a copy of the given asset and assigns the DNAAsset member variable to the copy.
+	 *
+	 * If SourceAsset is nullptr, DNAAsset will be set to nullptr.
+	 */
+	void CopyFromDNAAsset(const UDNAAsset* SourceAsset, UObject* OuterForOwnedObjects);
+
+	const UDNAAsset* GetDNAAsset() const { return DNAAsset; }
+
+	/** The index of the mesh component this DNA will be attached to */
+	UPROPERTY()
+	int32 ComponentIndex = 0;
+
+private:
+	/** Points to a DNA Asset that is owned by this struct */
+	UPROPERTY()
+	TObjectPtr<UDNAAsset> DNAAsset;
+};
+
+template<>
+struct TStructOpsTypeTraits<FDNAPinData> : public TStructOpsTypeTraitsBase2<FDNAPinData>
+{
+	enum
+	{
+		/** Tell Unreal's reflection system not to generate code that calls the copy constructor */
+		WithCopy = false
+	};
+};
+
+/** An extension for Mutable that allows users to bring RigLogic DNA into their Customizable Objects */
+UCLASS(MinimalAPI)
+class URigLogicMutableExtension : public UCustomizableObjectExtension
+{
+	GENERATED_BODY()
+
+public:
+	/** UCustomizableObjectExtension interface */
+	virtual TArray<FCustomizableObjectPinType> GetPinTypes() const override;
+	virtual TArray<FObjectNodeInputPin> GetAdditionalObjectNodePins() const override;
+	virtual void OnSkeletalMeshCreated(
+		const TArray<FInputPinDataContainer>& InputPinData,
+		int32 ComponentIndex,
+		USkeletalMesh* SkeletalMesh) const override;
+
+	/**
+	 * Makes a copy of the Source asset and returns it.
+	 *
+	 * The copy's Outer will be set to OuterForCopy.
+	 */
+	static UDNAAsset* CopyDNAAsset(const UDNAAsset* Source, UObject* OuterForCopy);
+
+	RIGLOGICMUTABLE_API static const FName DNAPinType;
+	RIGLOGICMUTABLE_API static const FName DNABaseNodePinName;
+	RIGLOGICMUTABLE_API static const FText DNANodeCategory;
+};
diff --git a/Engine/Plugins/Experimental/RigLogicMutable/Source/RigLogicMutable/Public/RigLogicMutableModule.h b/Engine/Plugins/Experimental/RigLogicMutable/Source/RigLogicMutable/Public/RigLogicMutableModule.h
new file mode 100644
--- /dev/null
+++ b/Engine/Plugins/Experimental/RigLogicMutable/Source/RigLogicMutable/Public/RigLogicMutableModule.h
@@ -0,0 +1,7 @@
+// Copyright Epic Games, Inc. All Rights Reserved.
+
+#pragma once
+
+#include "CoreMinimal.h"
+#include "Modules/ModuleInterface.h"
+#include "Modules/ModuleManager.h"
diff --git a/Engine/Plugins/Experimental/RigLogicMutable/Source/RigLogicMutable/RigLogicMutable.Build.cs b/Engine/Plugins/Experimental/RigLogicMutable/Source/RigLogicMutable/RigLogicMutable.Build.cs
new file mode 100644
--- /dev/null
+++ b/Engine/Plugins/Experimental/RigLogicMutable/Source/RigLogicMutable/RigLogicMutable.Build.cs
@@ -0,0 +1,25 @@
+// Copyright Epic Games, Inc. All Rights Reserved.
+
+namespace UnrealBuildTool.Rules
+{
+	public class RigLogicMutable : ModuleRules
+	{
+		public RigLogicMutable(ReadOnlyTargetRules Target) : base(Target)
+		{
+			PublicDependencyModuleNames.AddRange(
+				new string[]
+				{
+					"Core",
+					"CoreUObject",
+					"Engine",
+					"CustomizableObject",
+					"RigLogicModule",
+				});
+
+			PrivateDependencyModuleNames.AddRange(
+				new string[]
+				{
+				});
+		}
+	}
+}
diff --git a/Engine/Plugins/Experimental/RigLogicMutable/Source/RigLogicMutableEditor/Private/CustomizableObjectNodeDNAConstant.cpp b/Engine/Plugins/Experimental/RigLogicMutable/Source/RigLogicMutableEditor/Private/CustomizableObjectNodeDNAConstant.cpp
new file mode 100644
--- /dev/null
+++ b/Engine/Plugins/Experimental/RigLogicMutable/Source/RigLogicMutableEditor/Private/CustomizableObjectNodeDNAConstant.cpp
@@ -0,0 +1,92 @@
+// Copyright Epic Games, Inc. All Rights Reserved.
+
+#include "CustomizableObjectNodeDNAConstant.h"
+
+#include "DNAAsset.h"
+#include "Engine/SkeletalMesh.h"
+#include "MuCOE/EdGraphSchema_CustomizableObject.h"
+#include "MuCOE/ExtensionDataCompilerInterface.h"
+#include "MuCOE/RemapPins/CustomizableObjectNodeRemapPins.h"
+#include "MuR/ExtensionData.h"
+#include "MuT/NodeExtensionDataConstant.h"
+#include "RigLogicMutableExtension.h"
+
+#include UE_INLINE_GENERATED_CPP_BY_NAME(CustomizableObjectNodeDNAConstant)
+
+#define LOCTEXT_NAMESPACE "RigLogicMutableEditor"
+
+FText UCustomizableObjectNodeDNAConstant::GetNodeTitle(ENodeTitleType::Type TitleType) const
+{
+	return LOCTEXT("DNA_Constant", "DNA Constant");
+}
+
+FLinearColor UCustomizableObjectNodeDNAConstant::GetNodeTitleColor() const
+{
+	const UEdGraphSchema_CustomizableObject* Schema = GetDefault<UEdGraphSchema_CustomizableObject>();
+	return Schema->GetPinTypeColor(URigLogicMutableExtension::DNAPinType);
+}
+
+FText UCustomizableObjectNodeDNAConstant::GetTooltipText() const
+{
+	return LOCTEXT("DNA_Constant_Tooltip", "RigLogic DNA");
+}
+
+void UCustomizableObjectNodeDNAConstant::AllocateDefaultPins(UCustomizableObjectNodeRemapPins* RemapPins)
+{
+	UEdGraphPin* OutputPin = CustomCreatePin(EGPD_Output, URigLogicMutableExtension::DNAPinType, URigLogicMutableExtension::DNABaseNodePinName);
+	OutputPin->bDefaultValueIsIgnored = true;
+}
+
+bool UCustomizableObjectNodeDNAConstant::ShouldAddToContextMenu(FText& OutCategory) const
+{
+	OutCategory = UEdGraphSchema_CustomizableObject::NC_Experimental;
+	return true;
+}
+
+
+bool UCustomizableObjectNodeDNAConstant::IsExperimental() const
+{
+	return true;
+}
+
+
+mu::Ptr<mu::NodeExtensionData> UCustomizableObjectNodeDNAConstant::GenerateMutableNode(FExtensionDataCompilerInterface& CompilerInterface) const
+{
+	check(IsInGameThread());
+
+	if (SkeletalMesh)
+	{
+		CompilerInterface.AddParticipatingObject(*SkeletalMesh);
+	}
+	
+	// Create node and extension data container
+	mu::Ptr<mu::NodeExtensionDataConstant> Result = new mu::NodeExtensionDataConstant();
+	
+	// DNA is usually quite large, so set it up as a streaming constant to allow it to be loaded 
+	// on demand.
+	//
+	// If needed we could expose an editable UPROPERTY to give the user the option of making this
+	// an always-loaded constant.
+	UCustomizableObjectResourceDataContainer* Container = nullptr;
+	Result->SetValue(CompilerInterface.MakeStreamedExtensionData(Container));
+
+	// Populate instanced struct
+	if (Container)
+	{
+		FDNAPinData PinData;
+
+		if (SkeletalMesh)
+		{
+			// Note that this may be nullptr if the mesh doesn't have a DNA asset
+			PinData.CopyFromDNAAsset(Cast<UDNAAsset>(SkeletalMesh->GetAssetUserDataOfClass(UDNAAsset::StaticClass())), Container);
+		}
+
+		PinData.ComponentIndex = ComponentIndex;
+
+		Container->Data.Data.InitializeAs<FDNAPinData>(MoveTemp(PinData));
+	}
+
+	return Result;
+}
+
+#undef LOCTEXT_NAMESPACE
diff --git a/Engine/Plugins/Experimental/RigLogicMutable/Source/RigLogicMutableEditor/Private/RigLogicMutableEditorModule.cpp b/Engine/Plugins/Experimental/RigLogicMutable/Source/RigLogicMutableEditor/Private/RigLogicMutableEditorModule.cpp
new file mode 100644
--- /dev/null
+++ b/Engine/Plugins/Experimental/RigLogicMutable/Source/RigLogicMutableEditor/Private/RigLogicMutableEditorModule.cpp
@@ -0,0 +1,5 @@
+// Copyright Epic Games, Inc. All Rights Reserved.
+
+#include "RigLogicMutableEditorModule.h"
+
+IMPLEMENT_MODULE(FDefaultModuleImpl, RigLogicMutableEditor);
diff --git a/Engine/Plugins/Experimental/RigLogicMutable/Source/RigLogicMutableEditor/Public/CustomizableObjectNodeDNAConstant.h b/Engine/Plugins/Experimental/RigLogicMutable/Source/RigLogicMutableEditor/Public/CustomizableObjectNodeDNAConstant.h
new file mode 100644
--- /dev/null
+++ b/Engine/Plugins/Experimental/RigLogicMutable/Source/RigLogicMutableEditor/Public/CustomizableObjectNodeDNAConstant.h
@@ -0,0 +1,38 @@
+// Copyright Epic Games, Inc. All Rights Reserved.
+
+#pragma once
+
+#include "MuCOE/ICustomizableObjectExtensionNode.h"
+#include "MuCOE/Nodes/CustomizableObjectNodeExtensionDataConstant.h"
+
+#include "CustomizableObjectNodeDNAConstant.generated.h"
+
+/** Imports DNA from a Skeletal Mesh into the Customizable Object graph */
+UCLASS()
+class UCustomizableObjectNodeDNAConstant : public UCustomizableObjectNodeExtensionDataConstant, public ICustomizableObjectExtensionNode
+{
+	GENERATED_BODY()
+
+public:
+	/** EdGraphNode interface */
+	virtual FText GetNodeTitle(ENodeTitleType::Type TitleType) const override;
+	virtual FLinearColor GetNodeTitleColor() const override;
+	virtual FText GetTooltipText() const override;
+
+	/** UCustomizableObjectNode interface */
+	virtual void AllocateDefaultPins(UCustomizableObjectNodeRemapPins* RemapPins) override;
+	virtual bool ShouldAddToContextMenu(FText& OutCategory) const override;
+	virtual bool IsExperimental() const override;
+
+	/** ICustomizableObjectExtensionNode interface */
+	virtual mu::Ptr<mu::NodeExtensionData> GenerateMutableNode(FExtensionDataCompilerInterface& CompilerInterface) const override;
+
+	/** The Skeletal Mesh to copy DNA from */
+	UPROPERTY(EditAnywhere, Category = CustomizableObject)
+	TObjectPtr<USkeletalMesh> SkeletalMesh;
+
+	/** The index of the mesh component in the Customizable Object that this DNA will go to */
+	UPROPERTY(EditAnywhere, Category = CustomizableObject)
+	int32 ComponentIndex;
+};
+
diff --git a/Engine/Plugins/Experimental/RigLogicMutable/Source/RigLogicMutableEditor/Public/RigLogicMutableEditorModule.h b/Engine/Plugins/Experimental/RigLogicMutable/Source/RigLogicMutableEditor/Public/RigLogicMutableEditorModule.h
new file mode 100644
--- /dev/null
+++ b/Engine/Plugins/Experimental/RigLogicMutable/Source/RigLogicMutableEditor/Public/RigLogicMutableEditorModule.h
@@ -0,0 +1,7 @@
+// Copyright Epic Games, Inc. All Rights Reserved.
+
+#pragma once
+
+#include "CoreMinimal.h"
+#include "Modules/ModuleInterface.h"
+#include "Modules/ModuleManager.h"
diff --git a/Engine/Plugins/Experimental/RigLogicMutable/Source/RigLogicMutableEditor/RigLogicMutableEditor.Build.cs b/Engine/Plugins/Experimental/RigLogicMutable/Source/RigLogicMutableEditor/RigLogicMutableEditor.Build.cs
new file mode 100644
--- /dev/null
+++ b/Engine/Plugins/Experimental/RigLogicMutable/Source/RigLogicMutableEditor/RigLogicMutableEditor.Build.cs
@@ -0,0 +1,28 @@
+// Copyright Epic Games, Inc. All Rights Reserved.
+
+namespace UnrealBuildTool.Rules
+{
+	public class RigLogicMutableEditor : ModuleRules
+	{
+		public RigLogicMutableEditor(ReadOnlyTargetRules Target) : base(Target)
+		{
+			PublicDependencyModuleNames.AddRange(
+				new string[]
+				{
+					"Core",
+					"CoreUObject",
+					"Engine",
+					"CustomizableObjectEditor",
+					"MutableTools",
+					"MutableRuntime"
+				});
+
+			PrivateDependencyModuleNames.AddRange(
+				new string[]
+				{
+					"RigLogicModule",
+					"RigLogicMutable",
+				});
+		}
+	}
+}
