Added support for Mutable extensions to reference components by name instead of by index, and upgraded the RigLogicMutable extension to use this.

Also switched RigLogicMutable to using always-loaded extension data, to work around issues with streamed extension data.

diff --git a/Engine/Plugins/Experimental/RigLogicMutable/Source/RigLogicMutable/Private/RigLogicMutableExtension.cpp b/Engine/Plugins/Experimental/RigLogicMutable/Source/RigLogicMutable/Private/RigLogicMutableExtension.cpp
--- a/Engine/Plugins/Experimental/RigLogicMutable/Source/RigLogicMutable/Private/RigLogicMutableExtension.cpp
+++ b/Engine/Plugins/Experimental/RigLogicMutable/Source/RigLogicMutable/Private/RigLogicMutableExtension.cpp
@@ -21,7 +21,7 @@ FDNAPinData::FDNAPinData(FDNAPinData&& Source)
 
 FDNAPinData& FDNAPinData::operator=(FDNAPinData&& Source)
 {
-	ComponentIndex = Source.ComponentIndex;
+	ComponentName = Source.ComponentName;
 
 	DNAAsset = Source.DNAAsset;
 	Source.DNAAsset = nullptr;
@@ -68,7 +68,7 @@ TArray<FObjectNodeInputPin> URigLogicMutableExtension::GetAdditionalObjectNodePi
 
 void URigLogicMutableExtension::OnSkeletalMeshCreated(
 	const TArray<FInputPinDataContainer>& InputPinData,
-	int32 ComponentIndex,
+	FName ComponentName,
 	USkeletalMesh* SkeletalMesh) const
 {
 	// Find the DNA produced by the Customizable Object, if any, and assign it to the Skeletal Mesh
@@ -79,7 +79,7 @@ void URigLogicMutableExtension::OnSkeletalMeshCreated(
 		{
 			const FDNAPinData* Data = Container.Data.GetPtr<FDNAPinData>();
 			if (Data
-				&& Data->ComponentIndex == ComponentIndex
+				&& Data->ComponentName == ComponentName
 				&& Data->GetDNAAsset() != nullptr)
 			{
 				UDNAAsset* NewDNA = CopyDNAAsset(Data->GetDNAAsset(), SkeletalMesh);
diff --git a/Engine/Plugins/Experimental/RigLogicMutable/Source/RigLogicMutable/Public/RigLogicMutableExtension.h b/Engine/Plugins/Experimental/RigLogicMutable/Source/RigLogicMutable/Public/RigLogicMutableExtension.h
--- a/Engine/Plugins/Experimental/RigLogicMutable/Source/RigLogicMutable/Public/RigLogicMutableExtension.h
+++ b/Engine/Plugins/Experimental/RigLogicMutable/Source/RigLogicMutable/Public/RigLogicMutableExtension.h
@@ -40,9 +40,9 @@ public:
 
 	const UDNAAsset* GetDNAAsset() const { return DNAAsset; }
 
-	/** The index of the mesh component this DNA will be attached to */
+	/** The name of the mesh component this DNA will be attached to */
 	UPROPERTY()
-	int32 ComponentIndex = 0;
+	FName ComponentName;
 
 private:
 	/** Points to a DNA Asset that is owned by this struct */
@@ -72,7 +72,7 @@ public:
 	virtual TArray<FObjectNodeInputPin> GetAdditionalObjectNodePins() const override;
 	virtual void OnSkeletalMeshCreated(
 		const TArray<FInputPinDataContainer>& InputPinData,
-		int32 ComponentIndex,
+		FName ComponentName,
 		USkeletalMesh* SkeletalMesh) const override;
 
 	/**
diff --git a/Engine/Plugins/Experimental/RigLogicMutable/Source/RigLogicMutableEditor/Private/CustomizableObjectNodeDNAConstant.cpp b/Engine/Plugins/Experimental/RigLogicMutable/Source/RigLogicMutableEditor/Private/CustomizableObjectNodeDNAConstant.cpp
--- a/Engine/Plugins/Experimental/RigLogicMutable/Source/RigLogicMutableEditor/Private/CustomizableObjectNodeDNAConstant.cpp
+++ b/Engine/Plugins/Experimental/RigLogicMutable/Source/RigLogicMutableEditor/Private/CustomizableObjectNodeDNAConstant.cpp
@@ -61,29 +61,48 @@ mu::Ptr<mu::NodeExtensionData> UCustomizableObjectNodeDNAConstant::GenerateMutab
 	
 	// Create node and extension data container
 	mu::Ptr<mu::NodeExtensionDataConstant> Result = new mu::NodeExtensionDataConstant();
+
+	FDNAPinData PinData;
+	PinData.ComponentName = ComponentName;
 	
-	// DNA is usually quite large, so set it up as a streaming constant to allow it to be loaded 
-	// on demand.
+	// DNA is usually quite large, so ideally it would be made a streaming constant to allow it to 
+	// be loaded on demand.
 	//
-	// If needed we could expose an editable UPROPERTY to give the user the option of making this
-	// an always-loaded constant.
-	UCustomizableObjectResourceDataContainer* Container = nullptr;
-	Result->SetValue(CompilerInterface.MakeStreamedExtensionData(Container));
+	// However, streaming constants don't support subobjects properly at the moment, so we use an
+	// always-loaded constant instead.
+	const bool bShouldStreamDNA = false;
 
-	// Populate instanced struct
-	if (Container)
+	if (bShouldStreamDNA)
 	{
-		FDNAPinData PinData;
+		UCustomizableObjectResourceDataContainer* Container = nullptr;
+		Result->SetValue(CompilerInterface.MakeStreamedExtensionData(Container));
+
+		// Populate instanced struct
+		if (Container)
+		{
+			if (SkeletalMesh)
+			{
+				// Note that this may be nullptr if the mesh doesn't have a DNA asset
+				PinData.CopyFromDNAAsset(Cast<UDNAAsset>(SkeletalMesh->GetAssetUserDataOfClass(UDNAAsset::StaticClass())), Container);
+			}
 
+			Container->Data.Data.InitializeAs<FDNAPinData>(MoveTemp(PinData));
+		}
+	}
+	else
+	{
 		if (SkeletalMesh)
 		{
+			UObject* OuterForOwnedObjects = const_cast<UObject*>(CompilerInterface.GetOuterForAlwaysLoadedObjects());
+
 			// Note that this may be nullptr if the mesh doesn't have a DNA asset
-			PinData.CopyFromDNAAsset(Cast<UDNAAsset>(SkeletalMesh->GetAssetUserDataOfClass(UDNAAsset::StaticClass())), Container);
+			PinData.CopyFromDNAAsset(Cast<UDNAAsset>(SkeletalMesh->GetAssetUserDataOfClass(UDNAAsset::StaticClass())), OuterForOwnedObjects);
 		}
 
-		PinData.ComponentIndex = ComponentIndex;
+		FInstancedStruct Struct;
+		Struct.InitializeAs<FDNAPinData>(MoveTemp(PinData));
 
-		Container->Data.Data.InitializeAs<FDNAPinData>(MoveTemp(PinData));
+		Result->SetValue(CompilerInterface.MakeAlwaysLoadedExtensionData(MoveTemp(Struct)));
 	}
 
 	return Result;
diff --git a/Engine/Plugins/Experimental/RigLogicMutable/Source/RigLogicMutableEditor/Public/CustomizableObjectNodeDNAConstant.h b/Engine/Plugins/Experimental/RigLogicMutable/Source/RigLogicMutableEditor/Public/CustomizableObjectNodeDNAConstant.h
--- a/Engine/Plugins/Experimental/RigLogicMutable/Source/RigLogicMutableEditor/Public/CustomizableObjectNodeDNAConstant.h
+++ b/Engine/Plugins/Experimental/RigLogicMutable/Source/RigLogicMutableEditor/Public/CustomizableObjectNodeDNAConstant.h
@@ -31,8 +31,8 @@ public:
 	UPROPERTY(EditAnywhere, Category = CustomizableObject)
 	TObjectPtr<USkeletalMesh> SkeletalMesh;
 
-	/** The index of the mesh component in the Customizable Object that this DNA will go to */
+	/** The name of the mesh component in the Customizable Object that this DNA will go to */
 	UPROPERTY(EditAnywhere, Category = CustomizableObject)
-	int32 ComponentIndex;
+	FName ComponentName;
 };
 
diff --git a/Engine/Plugins/Mutable/Source/CustomizableObject/Private/MuCO/CustomizableObjectInstance.cpp b/Engine/Plugins/Mutable/Source/CustomizableObject/Private/MuCO/CustomizableObjectInstance.cpp
--- a/Engine/Plugins/Mutable/Source/CustomizableObject/Private/MuCO/CustomizableObjectInstance.cpp
+++ b/Engine/Plugins/Mutable/Source/CustomizableObject/Private/MuCO/CustomizableObjectInstance.cpp
@@ -2367,7 +2367,11 @@ bool UCustomizableInstancePrivate::UpdateSkeletalMesh_PostBeginUpdate0(UCustomiz
 		
 		for (const TPair<const UCustomizableObjectExtension*, TArray<FInputPinDataContainer>>& Pair : ExtensionToExtensionData)
 		{
+			PRAGMA_DISABLE_DEPRECATION_WARNINGS
 			Pair.Key->OnSkeletalMeshCreated(Pair.Value, ObjectComponentIndex, SkeletalMesh);
+			PRAGMA_ENABLE_DEPRECATION_WARNINGS
+
+			Pair.Key->OnSkeletalMeshCreated(Pair.Value, ComponentName, SkeletalMesh);
 		}
 		
 		// Mesh to copy data from if possible.
diff --git a/Engine/Plugins/Mutable/Source/CustomizableObject/Public/MuCO/CustomizableObjectExtension.h b/Engine/Plugins/Mutable/Source/CustomizableObject/Public/MuCO/CustomizableObjectExtension.h
--- a/Engine/Plugins/Mutable/Source/CustomizableObject/Public/MuCO/CustomizableObjectExtension.h
+++ b/Engine/Plugins/Mutable/Source/CustomizableObject/Public/MuCO/CustomizableObjectExtension.h
@@ -99,13 +99,16 @@ public:
 	 * @param InputPinData - The data for only the input pins *registered by this extension*. This
 	 * helps to enforce separation between the extensions, so that they don't depend on each other.
 	 * 
-	 * @param ObjectComponentIndex - The component index of the Skeletal Mesh, for the case where the pin
+	 * @param ComponentName - The component name of the Skeletal Mesh, for the case where the pin
 	 * data is associated with a particular component.
 	 * 
 	 * @param SkeletalMesh - The Skeletal Mesh that was created.
 	 */
-	virtual void OnSkeletalMeshCreated(const TArray<FInputPinDataContainer>& InputPinData, int32 ObjectComponentIndex, USkeletalMesh* SkeletalMesh) const {}
+	virtual void OnSkeletalMeshCreated(const TArray<FInputPinDataContainer>& InputPinData, FName ComponentName, USkeletalMesh* SkeletalMesh) const {}
 	
+	UE_DEPRECATED(5.6, "This overload of OnSkeletalMeshCreated is deprecated and will be removed. Use ComponentName instead.")
+	virtual void OnSkeletalMeshCreated(const TArray<FInputPinDataContainer>& InputPinData, int32 ObjectComponentIndex, USkeletalMesh* SkeletalMesh) const {}
+
 	/**
 	 * Note that the data registered here is completely independent of any Extension Data used in
 	 * the Customizable Object graph. Even though Extension Data and this Extension Instance Data
