[RigLogicMutable] Fixed Game crash when generating a SKM with DNA.

#rnx

diff --git a/Engine/Plugins/Experimental/RigLogicMutable/Source/RigLogicMutable/Private/RigLogicMutableExtension.cpp b/Engine/Plugins/Experimental/RigLogicMutable/Source/RigLogicMutable/Private/RigLogicMutableExtension.cpp
--- a/Engine/Plugins/Experimental/RigLogicMutable/Source/RigLogicMutable/Private/RigLogicMutableExtension.cpp
+++ b/Engine/Plugins/Experimental/RigLogicMutable/Source/RigLogicMutable/Private/RigLogicMutableExtension.cpp
@@ -34,6 +34,7 @@ void FDNAPinData::CopyFromDNAAsset(const UDNAAsset* SourceAsset, UObject* OuterF
 	if (SourceAsset)
 	{
 		DNAAsset = URigLogicMutableExtension::CopyDNAAsset(SourceAsset, OuterForOwnedObjects);
+		DNAAsset->bKeepDNAAfterInitialization = true;
 	}
 	else
 	{
@@ -97,27 +98,5 @@ UDNAAsset* URigLogicMutableExtension::CopyDNAAsset(const UDNAAsset* Source, UObj
 {
 	check(IsInGameThread());
 
-	// Currently the only way to copy a UDNAAsset is to serialize it into a buffer and deserialize
-	// the buffer into the copy.
-
-	// Serialize the existing DNA into a buffer
-	TArray<uint8> Buffer;
-	{
-		FMemoryWriter Writer(Buffer);
-
-		// Need to const_cast because Serialize is non-const. Source will not be modified.
-		const_cast<UDNAAsset*>(Source)->Serialize(Writer);
-
-		Writer.Close();
-	}
-
-	// Create the new DNA asset and deserialize the buffer into it
-	UDNAAsset* Result = NewObject<UDNAAsset>(OuterForCopy);
-	{
-		FMemoryReader Reader(Buffer);
-
-		Result->Serialize(Reader);
-	}
-
-	return Result;
+	return DuplicateObject(Source, OuterForCopy);
 }
