[mutable] Fix compilation for plugin after removal of AddParticipatingObjects method.
The reference should be picked from the node anyway with the new method, but if it was a soft object ptr we would need a virtual method in the node to manually add it.
#rnx

diff --git a/Engine/Plugins/Experimental/RigLogicMutable/Source/RigLogicMutableEditor/Private/CustomizableObjectNodeDNAConstant.cpp b/Engine/Plugins/Experimental/RigLogicMutable/Source/RigLogicMutableEditor/Private/CustomizableObjectNodeDNAConstant.cpp
--- a/Engine/Plugins/Experimental/RigLogicMutable/Source/RigLogicMutableEditor/Private/CustomizableObjectNodeDNAConstant.cpp
+++ b/Engine/Plugins/Experimental/RigLogicMutable/Source/RigLogicMutableEditor/Private/CustomizableObjectNodeDNAConstant.cpp
@@ -53,11 +53,6 @@ bool UCustomizableObjectNodeDNAConstant::IsExperimental() const
 mu::Ptr<mu::NodeExtensionData> UCustomizableObjectNodeDNAConstant::GenerateMutableNode(FExtensionDataCompilerInterface& CompilerInterface) const
 {
 	check(IsInGameThread());
-
-	if (SkeletalMesh)
-	{
-		CompilerInterface.AddParticipatingObject(*SkeletalMesh);
-	}
 	
 	// Create node and extension data container
 	mu::Ptr<mu::NodeExtensionDataConstant> Result = new mu::NodeExtensionDataConstant();
